Audio Platform Settings UX cleanup

diff --git a/Engine/Source/Developer/Windows/WindowsTargetPlatform/Private/WindowsTargetPlatformClasses.cpp b/Engine/Source/Developer/Windows/WindowsTargetPlatform/Private/WindowsTargetPlatformClasses.cpp
--- a/Engine/Source/Developer/Windows/WindowsTargetPlatform/Private/WindowsTargetPlatformClasses.cpp
+++ b/Engine/Source/Developer/Windows/WindowsTargetPlatform/Private/WindowsTargetPlatformClasses.cpp
@@ -7,6 +7,13 @@
 
 UWindowsTargetSettings::UWindowsTargetSettings( const FObjectInitializer& ObjectInitializer )
 	: Super(ObjectInitializer)
+	, CacheSizeKB(65536)
+	, MaxSampleRate(48000)
+	, HighSampleRate(32000)
+	, MedSampleRate(24000)
+	, LowSampleRate(12000)
+	, MinSampleRate(8000)
+	, CompressionQualityModifier(1)
 {
 	// Default windows settings
 	AudioSampleRate = 48000;
diff --git a/Engine/Source/Runtime/Android/AndroidRuntimeSettings/Private/AndroidRuntimeSettings.cpp b/Engine/Source/Runtime/Android/AndroidRuntimeSettings/Private/AndroidRuntimeSettings.cpp
--- a/Engine/Source/Runtime/Android/AndroidRuntimeSettings/Private/AndroidRuntimeSettings.cpp
+++ b/Engine/Source/Runtime/Android/AndroidRuntimeSettings/Private/AndroidRuntimeSettings.cpp
@@ -27,6 +27,13 @@ UAndroidRuntimeSettings::UAndroidRuntimeSettings(const FObjectInitializer& Objec
 	, AudioSampleRate(44100)
 	, AudioCallbackBufferFrameSize(1024)
 	, AudioNumBuffersToEnqueue(4)
+	, CacheSizeKB(65536)
+	, MaxSampleRate(48000)
+	, HighSampleRate(32000)
+    , MedSampleRate(24000)
+    , LowSampleRate(12000)
+	, MinSampleRate(8000)
+	, CompressionQualityModifier(1)
 	, bMultiTargetFormat_ETC2(true)
 	, bMultiTargetFormat_DXT(true)
 	, bMultiTargetFormat_ASTC(true)
diff --git a/Engine/Source/Runtime/AudioMixerCore/Public/AudioMixerTypes.h b/Engine/Source/Runtime/AudioMixerCore/Public/AudioMixerTypes.h
--- a/Engine/Source/Runtime/AudioMixerCore/Public/AudioMixerTypes.h
+++ b/Engine/Source/Runtime/AudioMixerCore/Public/AudioMixerTypes.h
@@ -71,7 +71,7 @@ struct AUDIOMIXERCORE_API FAudioPlatformSettings
 	/** The number of buffers to keep enqueued. More buffers increases latency, but can compensate for variable compute availability in audio callbacks on some platforms. */
 	int32 NumBuffers;
 
-	/** The max number of channels to limit for this platform. The max channels used will be the minimum of this value and the global audio quality settings. A value of 0 will not apply a platform channel count max. */
+	/** The max number of channels (simultaneous voices) to use as the limit for this platform. If given a value of 0, it will use the value from the active Global Audio Quality Settings */
 	int32 MaxChannels;
 
 	/** The number of workers to use to compute source audio. Will only use up to the max number of sources. Will evenly divide sources to each source worker. */
diff --git a/Engine/Source/Runtime/AudioPlatformConfiguration/Public/AudioCompressionSettings.h b/Engine/Source/Runtime/AudioPlatformConfiguration/Public/AudioCompressionSettings.h
--- a/Engine/Source/Runtime/AudioPlatformConfiguration/Public/AudioCompressionSettings.h
+++ b/Engine/Source/Runtime/AudioPlatformConfiguration/Public/AudioCompressionSettings.h
@@ -149,23 +149,24 @@ struct AUDIOPLATFORMCONFIGURATION_API FPlatformRuntimeAudioCompressionOverrides
 {
 	GENERATED_USTRUCT_BODY()
 
-	// Set this to true to override Sound Groups and use the Duration Threshold value to determine whether a sound should be fully decompressed during initial loading.
+	// When true, overrides the Sound Group on each Sound Wave, and instead uses the Duration Threshold value to determine whether a sound should be fully decompressed during initial loading.
 	UPROPERTY(EditAnywhere, Category = "DecompressOnLoad")
 	bool bOverrideCompressionTimes;
 	
 	// When Override Compression Times is set to true, any sound under this threshold (in seconds) will be fully decompressed on load.
 	// Otherwise the first chunk of this sound is cached at load and the rest is decompressed in real time.
+	// If set to zero, will default to the Sound Group on the relevant Sound Wave
 	UPROPERTY(EditAnywhere, Category = "DecompressOnLoad")
 	float DurationThreshold;
 
 	// On this platform, any random nodes on Sound Cues will automatically only preload this number of branches and dispose of any others
-	// on load. This can drastically cut down on memory usage.
-	UPROPERTY(EditAnywhere, Category = "SoundCueLoading", meta = (DisplayName = "Maximum Branches on Random SoundCue nodes", ClampMin = "1"))
+	// on load. This can drastically cut down on memory usage. If set to 0, no branches are culled.
+	UPROPERTY(EditAnywhere, Category = "SoundCueLoading", meta = (DisplayName = "Maximum Branches on Random SoundCue nodes", ClampMin = "0"))
 	int32 MaxNumRandomBranches;
 
 	// On this platform, use the specified quality at this index to override the quality used for SoundCues on this platform
-	UPROPERTY(EditAnywhere, Category = "SoundCueLoading", meta = (DisplayName = "Quality Index for Sound Cues", ClampMin = "-1", ClampMax = "50"))
-	int32 SoundCueQualityIndex;
+	UPROPERTY(EditAnywhere, Category = "SoundCueLoading", meta = (DisplayName = "Quality Index for Sound Cues", ClampMin = "0", ClampMax = "50"))
+	int32 SoundCueQualityIndex = 1;
 
 	FPlatformRuntimeAudioCompressionOverrides();
 
diff --git a/Engine/Source/Runtime/IOS/IOSRuntimeSettings/Private/IOSRuntimeSettings.cpp b/Engine/Source/Runtime/IOS/IOSRuntimeSettings/Private/IOSRuntimeSettings.cpp
--- a/Engine/Source/Runtime/IOS/IOSRuntimeSettings/Private/IOSRuntimeSettings.cpp
+++ b/Engine/Source/Runtime/IOS/IOSRuntimeSettings/Private/IOSRuntimeSettings.cpp
@@ -9,6 +9,13 @@
 
 UIOSRuntimeSettings::UIOSRuntimeSettings(const FObjectInitializer& ObjectInitializer)
 	: Super(ObjectInitializer)
+	, CacheSizeKB(65536)
+	, MaxSampleRate(48000)
+	, HighSampleRate(32000)
+	, MedSampleRate(24000)
+	, LowSampleRate(12000)
+	, MinSampleRate(8000)
+	, CompressionQualityModifier(1)
 {
 	bEnableGameCenterSupport = true;
 	bEnableCloudKitSupport = false;
