Fix output buffer size calculation for audio resampler.
#rnx

diff --git a/Engine/Source/Runtime/AudioPlatformConfiguration/Private/AudioResampler.cpp b/Engine/Source/Runtime/AudioPlatformConfiguration/Private/AudioResampler.cpp
--- a/Engine/Source/Runtime/AudioPlatformConfiguration/Private/AudioResampler.cpp
+++ b/Engine/Source/Runtime/AudioPlatformConfiguration/Private/AudioResampler.cpp
@@ -41,8 +41,12 @@ namespace Audio
 
 	int32 GetOutputBufferSize(const FResamplingParameters& InParameters)
 	{
+		const int32 NumChannels = FMath::Max(1, InParameters.NumChannels);
+		const int32 NumInputFrames = InParameters.InputBuffer.Num() / NumChannels;
 		const float Ratio = InParameters.DestinationSampleRate / InParameters.SourceSampleRate;
-		return InParameters.InputBuffer.Num() * Ratio;
+		const int32 NumOutputFrames = FMath::CeilToInt(Ratio * NumInputFrames);
+
+		return NumChannels * NumOutputFrames;
 	}
 
 	bool Resample(const FResamplingParameters& InParameters, FResamplerResults& OutData)
