Default AUDIO_BUFFER_ALIGNMENT to 16 on all platforms.

The platform that required a higher alignment is using a software implementation. Also, the alignment is handled at a lower level when a hardware implementation is used.

diff --git a/Engine/Source/Runtime/AudioPlatformConfiguration/Public/AudioResampler.h b/Engine/Source/Runtime/AudioPlatformConfiguration/Public/AudioResampler.h
--- a/Engine/Source/Runtime/AudioPlatformConfiguration/Public/AudioResampler.h
+++ b/Engine/Source/Runtime/AudioPlatformConfiguration/Public/AudioResampler.h
@@ -8,14 +8,9 @@
 #include "Logging/LogMacros.h"
 #include "Templates/UniquePtr.h"
 
-#if PLATFORM_SWITCH
-// Switch uses page alignment for submitted buffers
 // TODO: Move BufferVectorOperations AUDIO_BUFFER_ALIGNMENT
 // define to more central location and reference here.
-#define AUDIO_BUFFER_ALIGNMENT 4096
-#else
 #define AUDIO_BUFFER_ALIGNMENT 16
-#endif
 
 DECLARE_LOG_CATEGORY_EXTERN(LogAudioResampler, Warning, All);
 
diff --git a/Engine/Source/Runtime/SignalProcessing/Public/DSP/BufferVectorOperations.h b/Engine/Source/Runtime/SignalProcessing/Public/DSP/BufferVectorOperations.h
--- a/Engine/Source/Runtime/SignalProcessing/Public/DSP/BufferVectorOperations.h
+++ b/Engine/Source/Runtime/SignalProcessing/Public/DSP/BufferVectorOperations.h
@@ -7,7 +7,6 @@
 
 #if UE_ENABLE_INCLUDE_ORDER_DEPRECATED_IN_5_1 
 #include "CoreMinimal.h"
-#endif
 
 namespace Audio
 {
