Change Bink Quality mapping to use full 1-100 range. Add Project Setting to change "Default Compression Quality". Rewrite of DDC key hash to fix collisions.
NOTE: Bumping Bink version as mapping changed. Also bumping Streamcache versions as hash function changed.

diff --git a/Engine/Source/Developer/AudioFormatBink/Private/AudioFormatBink.cpp b/Engine/Source/Developer/AudioFormatBink/Private/AudioFormatBink.cpp
--- a/Engine/Source/Developer/AudioFormatBink/Private/AudioFormatBink.cpp
+++ b/Engine/Source/Developer/AudioFormatBink/Private/AudioFormatBink.cpp
@@ -20,11 +20,15 @@ namespace AudioFormatBinkPrivate
 		static const float BinkLowest = 4;
 		static const float BinkHighest = 0;
 
-		// Map Quality 1 (lowest) to 40 (highest).
+		// Map Quality 1 (lowest) to 100 (highest).
 		static const float QualityLowest = 1;
-		static const float QualityHighest = 40;
+		static const float QualityHighest = 100;
 
-		return FMath::GetMappedRangeValueClamped(FVector2D(QualityLowest, QualityHighest), FVector2D(BinkLowest, BinkHighest), InQualityIndex);
+		// Map Quality into Bink Range. Note: +1 gives the Bink range 5 steps inclusive.
+		float BinkValue = FMath::GetMappedRangeValueClamped(FVector2D(QualityLowest, QualityHighest), FVector2D(BinkLowest + 1.f, BinkHighest), InQualityIndex);
+
+		// Floor each value and clamp into range (as top lerp will be +1 over)
+		return FMath::Clamp(FMath::FloorToInt(BinkValue), BinkHighest, BinkLowest);
 	}	
 }
 
diff --git a/Engine/Source/Runtime/AudioPlatformConfiguration/Private/AudioCompressionSettings.cpp b/Engine/Source/Runtime/AudioPlatformConfiguration/Private/AudioCompressionSettings.cpp
--- a/Engine/Source/Runtime/AudioPlatformConfiguration/Private/AudioCompressionSettings.cpp
+++ b/Engine/Source/Runtime/AudioPlatformConfiguration/Private/AudioCompressionSettings.cpp
@@ -20,5 +20,55 @@ FPlatformRuntimeAudioCompressionOverrides* FPlatformRuntimeAudioCompressionOverr
 // absurdly large number, to ensure you do not pollute the DDC.
 int32 FPlatformAudioCookOverrides::GetStreamCachingVersion()
 {
-	return 5027;
+	return 5028;
+}
+namespace PlatformAudioCookOverridesPrivate
+{
+	template<typename HashType>
+	void AppendHash(FString& OutString, const TCHAR* InName, const HashType& InValueToHash)
+	{
+		FString Lexed = LexToString(InValueToHash); // Using temporary prevents TCHAR/String mismatch in printf
+		OutString += FString::Printf(TEXT("%s_%s_"), InName, *Lexed);
+	}
+}
+
+void FPlatformAudioCookOverrides::GetHashSuffix(const FPlatformAudioCookOverrides* InOverrides, FString& OutSuffix)
+{
+	if (InOverrides == nullptr)
+	{
+		return;
+	}	
+	
+	using namespace PlatformAudioCookOverridesPrivate;
+
+	// Starting Delim is important, as FSoundWaveData::FindRuntimeFormat, uses it determine format from the inline chunk name.
+	OutSuffix += TEXT("_");
+
+	// Start with StreamCache version.
+	AppendHash(OutSuffix, TEXT("SCVER"), GetStreamCachingVersion());
+	
+	// Each member in declaration order.
+	
+	// FPlatformAudioCookOverrides
+	AppendHash(OutSuffix, TEXT("R4DV"), InOverrides->bResampleForDevice);
+
+	TArray<float> Rates;
+	InOverrides->PlatformSampleRates.GenerateValueArray(Rates);
+	for (int32 i = 0; i < Rates.Num(); ++i)
+	{
+		AppendHash(OutSuffix, *FString::Printf(TEXT("SR%d"), i), Rates[i]);
+	}
+
+	AppendHash(OutSuffix, TEXT("QMOD"), InOverrides->CompressionQualityModifier);
+	AppendHash(OutSuffix, TEXT("CQLT"), InOverrides->SoundCueCookQualityIndex);
+	AppendHash(OutSuffix, TEXT("ASTH"), InOverrides->AutoStreamingThreshold);
+	AppendHash(OutSuffix, TEXT("INLC"), InOverrides->bInlineStreamedAudioChunks);		
+
+	// FAudioStreamCachingSettings
+	AppendHash(OutSuffix, TEXT("CSZE"), InOverrides->StreamCachingSettings.CacheSizeKB);
+	AppendHash(OutSuffix, TEXT("LCF"), InOverrides->StreamCachingSettings.bForceLegacyStreamChunking);
+	AppendHash(OutSuffix, TEXT("ZCS"), InOverrides->StreamCachingSettings.ZerothChunkSizeForLegacyStreamChunkingKB);
+	AppendHash(OutSuffix, TEXT("MCSO"), InOverrides->StreamCachingSettings.MaxChunkSizeOverrideKB);
+	
+	OutSuffix += TEXT("END");
 }
diff --git a/Engine/Source/Runtime/AudioPlatformConfiguration/Public/AudioCompressionSettings.h b/Engine/Source/Runtime/AudioPlatformConfiguration/Public/AudioCompressionSettings.h
--- a/Engine/Source/Runtime/AudioPlatformConfiguration/Public/AudioCompressionSettings.h
+++ b/Engine/Source/Runtime/AudioPlatformConfiguration/Public/AudioCompressionSettings.h
@@ -58,12 +58,12 @@ struct FAudioStreamCachingSettings
 /* This struct is used for settings used during the cook to a target    */
 /* platform (platform-specific compression quality and resampling, etc.)*/
 /************************************************************************/
-struct FPlatformAudioCookOverrides
+struct AUDIOPLATFORMCONFIGURATION_API FPlatformAudioCookOverrides
 {
 	// Increment this return value to force a recook on all Stream Caching assets.
 	// For testing, it's useful to set this to either a negative number or
 	// absurdly large number, to ensure you do not pollute the DDC.
-	static AUDIOPLATFORMCONFIGURATION_API int32 GetStreamCachingVersion();
+	static int32 GetStreamCachingVersion();
 
 	bool bResampleForDevice;
 
@@ -100,48 +100,7 @@ struct FPlatformAudioCookOverrides
 	}
 
 	// This is used to invalidate compressed audio for a specific platform.
-	static void GetHashSuffix(const FPlatformAudioCookOverrides* InOverrides, FString& OutSuffix)
-	{
-		if (InOverrides == nullptr)
-		{
-			return;
-		}
-		OutSuffix.AppendChar('_');
-
-		int32 CompressionQualityHash = FMath::FloorToInt(InOverrides->CompressionQualityModifier * 100.0f);
-		OutSuffix.AppendInt(CompressionQualityHash);
-
-		int32 AutoStreamingThresholdHash = FMath::FloorToInt(InOverrides->AutoStreamingThreshold * 100.0f);
-		OutSuffix.AppendInt(AutoStreamingThresholdHash);
-
-		OutSuffix.Append(TEXT("_StreamCache_Ver"));
-		OutSuffix.AppendInt(GetStreamCachingVersion());
-		OutSuffix.AppendChar('_');
-
-		// cache info:
-		OutSuffix.Append(TEXT("MEM_"));
-		OutSuffix.AppendInt(InOverrides->StreamCachingSettings.CacheSizeKB);
-		OutSuffix.Append(TEXT("MaxChnkSize_"));
-		OutSuffix.AppendInt(InOverrides->StreamCachingSettings.MaxChunkSizeOverrideKB);
-
-		if (InOverrides->StreamCachingSettings.bForceLegacyStreamChunking)
-		{
-			OutSuffix.Append(TEXT("_LegacyChunking_"));
-			OutSuffix.AppendInt(InOverrides->StreamCachingSettings.ZerothChunkSizeForLegacyStreamChunkingKB);
-		}
-		
-
-		int32 ResampleBoolHash = (int32)InOverrides->bResampleForDevice;
-		OutSuffix.AppendInt(ResampleBoolHash);
-
-		TMap<ESoundwaveSampleRateSettings, float> SampleRateMap = InOverrides->PlatformSampleRates;
-
-		for (auto& SampleRateQuality : SampleRateMap)
-		{
-			int32 SampleRateHash = FMath::FloorToInt(SampleRateQuality.Value / 1000.0f);
-			OutSuffix.AppendInt(SampleRateHash);
-		}
-	}
+	static void GetHashSuffix(const FPlatformAudioCookOverrides* InOverrides, FString& OutSuffix);
 };
 
 USTRUCT()
diff --git a/Engine/Source/Runtime/Engine/Classes/Sound/AudioSettings.h b/Engine/Source/Runtime/Engine/Classes/Sound/AudioSettings.h
--- a/Engine/Source/Runtime/Engine/Classes/Sound/AudioSettings.h
+++ b/Engine/Source/Runtime/Engine/Classes/Sound/AudioSettings.h
@@ -170,6 +170,10 @@ class ENGINE_API UAudioSettings : public UDeveloperSettings
 	UPROPERTY(config, EditAnywhere, Category = "Audio")
 	EDefaultAudioCompressionType DefaultAudioCompressionType;
 
+	/** The default compression quality (e.g. for new SoundWaves) */
+	UPROPERTY(config, EditAnywhere, Category = "Audio", meta = (ClampMin = 1, UIMin = 1, UIMax = 100))
+	int32 DefaultCompressionQuality = 80;
+
 	/** The amount of audio to send to reverb submixes if no reverb send is setup for the source through attenuation settings. Only used in audio mixer. */
 	UPROPERTY(config)
 	float DefaultReverbSendLevel_DEPRECATED;
@@ -272,6 +276,8 @@ public:
 
 	// Get the quality level settings at the provided level index
 	const FAudioQualitySettings& GetQualityLevelSettings(int32 QualityLevel) const;
+
+	int32 GetDefaultCompressionQuality() const;
 	
 	// Get the quality name level for a given index
 	FString FindQualityNameByIndex(int32 Index) const;
diff --git a/Engine/Source/Runtime/Engine/Private/AudioDerivedData.cpp b/Engine/Source/Runtime/Engine/Private/AudioDerivedData.cpp
--- a/Engine/Source/Runtime/Engine/Private/AudioDerivedData.cpp
+++ b/Engine/Source/Runtime/Engine/Private/AudioDerivedData.cpp
@@ -164,7 +164,7 @@ Derived data key generation.
 
 // If you want to bump this version, generate a new guid using
 // VS->Tools->Create GUID and paste it here. https://www.guidgen.com works too.
-#define STREAMEDAUDIO_DERIVEDDATA_VER		TEXT("0E2C7011AFE845488C383F8A3531A3F1")
+#define STREAMEDAUDIO_DERIVEDDATA_VER		TEXT("E9A61FD4C181422A8EAAB060B1690F6A")
 
 /**
  * Computes the derived data key suffix for a SoundWave's Streamed Audio.
diff --git a/Engine/Source/Runtime/Engine/Private/AudioSettings.cpp b/Engine/Source/Runtime/Engine/Private/AudioSettings.cpp
--- a/Engine/Source/Runtime/Engine/Private/AudioSettings.cpp
+++ b/Engine/Source/Runtime/Engine/Private/AudioSettings.cpp
@@ -188,6 +188,11 @@ const FAudioQualitySettings& UAudioSettings::GetQualityLevelSettings(int32 Quali
 	return QualityLevels[FMath::Clamp(QualityLevel, 0, QualityLevels.Num() - 1)];
 }
 
+int32 UAudioSettings::GetDefaultCompressionQuality() const
+{
+	return FMath::Clamp(DefaultCompressionQuality,1,100);
+}
+
 int32 UAudioSettings::GetQualityLevelSettingsNum() const
 {
 	return QualityLevels.Num();
diff --git a/Engine/Source/Runtime/Engine/Private/SoundWave.cpp b/Engine/Source/Runtime/Engine/Private/SoundWave.cpp
--- a/Engine/Source/Runtime/Engine/Private/SoundWave.cpp
+++ b/Engine/Source/Runtime/Engine/Private/SoundWave.cpp
@@ -692,10 +692,9 @@ uint32 FStreamedAudioChunk::StoreInDerivedDataCache(const FString& InDerivedData
 USoundWave::USoundWave(const FObjectInitializer& ObjectInitializer)
 	: Super(ObjectInitializer)
 {
-
 	Volume = 1.0;
 	Pitch = 1.0;
-	CompressionQuality = 40;
+	CompressionQuality = -1; // Note: This is set in PostInitProperties, as its unsafe to GetDefault<UAudioSettings> here.
 	SubtitlePriority = DEFAULT_SUBTITLE_PRIORITY;
 	ResourceState = ESoundWaveResourceState::NeedsFree;
 	RawPCMDataSize = 0;
@@ -1145,6 +1144,9 @@ void USoundWave::PostInitProperties()
 {
 	Super::PostInitProperties();
 
+	// Setup the default Compression Quality, ahead of Serialization
+	CompressionQuality = GetDefault<UAudioSettings>()->GetDefaultCompressionQuality();
+
 	if(!IsTemplate())
 	{
 		// Don't rebuild our streaming chunks yet because we may not have loaded the RawPCMData at this point.
