Fix DDC issue with Compression Quality

diff --git a/Engine/Source/Runtime/AudioPlatformConfiguration/Private/AudioCompressionSettings.cpp b/Engine/Source/Runtime/AudioPlatformConfiguration/Private/AudioCompressionSettings.cpp
--- a/Engine/Source/Runtime/AudioPlatformConfiguration/Private/AudioCompressionSettings.cpp
+++ b/Engine/Source/Runtime/AudioPlatformConfiguration/Private/AudioCompressionSettings.cpp
@@ -1,6 +1,7 @@
 // Copyright Epic Games, Inc. All Rights Reserved.
 
 #include "AudioCompressionSettings.h"
+#include "AudioCompressionSettingsUtils.h"
 
 #include UE_INLINE_GENERATED_CPP_BY_NAME(AudioCompressionSettings)
 
@@ -22,15 +23,7 @@ int32 FPlatformAudioCookOverrides::GetStreamCachingVersion()
 {
 	return 5028;
 }
-namespace PlatformAudioCookOverridesPrivate
-{
-	template<typename HashType>
-	void AppendHash(FString& OutString, const TCHAR* InName, const HashType& InValueToHash)
-	{
-		FString Lexed = LexToString(InValueToHash); // Using temporary prevents TCHAR/String mismatch in printf
-		OutString += FString::Printf(TEXT("%s_%s_"), InName, *Lexed);
-	}
-}
+
 
 void FPlatformAudioCookOverrides::GetHashSuffix(const FPlatformAudioCookOverrides* InOverrides, FString& OutSuffix)
 {
@@ -39,36 +32,37 @@ void FPlatformAudioCookOverrides::GetHashSuffix(const FPlatformAudioCookOverride
 		return;
 	}	
 	
-	using namespace PlatformAudioCookOverridesPrivate;
+	using FPCU = FPlatformCompressionUtilities;
+	FString SoundWaveHash;
 
 	// Starting Delim is important, as FSoundWaveData::FindRuntimeFormat, uses it determine format from the inline chunk name.
 	OutSuffix += TEXT("_");
 
 	// Start with StreamCache version.
-	AppendHash(OutSuffix, TEXT("SCVER"), GetStreamCachingVersion());
+	FPCU::AppendHash(OutSuffix, TEXT("SCVER"), GetStreamCachingVersion());
 	
 	// Each member in declaration order.
 	
 	// FPlatformAudioCookOverrides
-	AppendHash(OutSuffix, TEXT("R4DV"), InOverrides->bResampleForDevice);
+	FPCU::AppendHash(OutSuffix, TEXT("R4DV"), InOverrides->bResampleForDevice);
 
 	TArray<float> Rates;
 	InOverrides->PlatformSampleRates.GenerateValueArray(Rates);
 	for (int32 i = 0; i < Rates.Num(); ++i)
 	{
-		AppendHash(OutSuffix, *FString::Printf(TEXT("SR%d"), i), Rates[i]);
+		FPCU::AppendHash(OutSuffix, *FString::Printf(TEXT("SR%d"), i), Rates[i]);
 	}
 
-	AppendHash(OutSuffix, TEXT("QMOD"), InOverrides->CompressionQualityModifier);
-	AppendHash(OutSuffix, TEXT("CQLT"), InOverrides->SoundCueCookQualityIndex);
-	AppendHash(OutSuffix, TEXT("ASTH"), InOverrides->AutoStreamingThreshold);
-	AppendHash(OutSuffix, TEXT("INLC"), InOverrides->bInlineStreamedAudioChunks);		
+	FPCU::AppendHash(OutSuffix, TEXT("QMOD"), InOverrides->CompressionQualityModifier);
+	FPCU::AppendHash(OutSuffix, TEXT("CQLT"), InOverrides->SoundCueCookQualityIndex);
+	FPCU::AppendHash(OutSuffix, TEXT("ASTH"), InOverrides->AutoStreamingThreshold);
+	FPCU::AppendHash(OutSuffix, TEXT("INLC"), InOverrides->bInlineStreamedAudioChunks);
 
 	// FAudioStreamCachingSettings
-	AppendHash(OutSuffix, TEXT("CSZE"), InOverrides->StreamCachingSettings.CacheSizeKB);
-	AppendHash(OutSuffix, TEXT("LCF"), InOverrides->StreamCachingSettings.bForceLegacyStreamChunking);
-	AppendHash(OutSuffix, TEXT("ZCS"), InOverrides->StreamCachingSettings.ZerothChunkSizeForLegacyStreamChunkingKB);
-	AppendHash(OutSuffix, TEXT("MCSO"), InOverrides->StreamCachingSettings.MaxChunkSizeOverrideKB);
+	FPCU::AppendHash(OutSuffix, TEXT("CSZE"), InOverrides->StreamCachingSettings.CacheSizeKB);
+	FPCU::AppendHash(OutSuffix, TEXT("LCF"), InOverrides->StreamCachingSettings.bForceLegacyStreamChunking);
+	FPCU::AppendHash(OutSuffix, TEXT("ZCS"), InOverrides->StreamCachingSettings.ZerothChunkSizeForLegacyStreamChunkingKB);
+	FPCU::AppendHash(OutSuffix, TEXT("MCSO"), InOverrides->StreamCachingSettings.MaxChunkSizeOverrideKB);
 	
 	OutSuffix += TEXT("END");
 }
diff --git a/Engine/Source/Runtime/Engine/Private/AudioDerivedData.cpp b/Engine/Source/Runtime/Engine/Private/AudioDerivedData.cpp
--- a/Engine/Source/Runtime/Engine/Private/AudioDerivedData.cpp
+++ b/Engine/Source/Runtime/Engine/Private/AudioDerivedData.cpp
@@ -19,6 +19,7 @@
 #include "Sound/SoundWaveProcedural.h"
 #include "DSP/FloatArrayMath.h"
 #include "DSP/MultichannelBuffer.h"
+#include "Containers/UnrealString.h"
 
 DEFINE_LOG_CATEGORY_STATIC(LogAudioDerivedData, Log, All);
 
@@ -166,6 +167,19 @@ Derived data key generation.
 // VS->Tools->Create GUID and paste it here. https://www.guidgen.com works too.
 #define STREAMEDAUDIO_DERIVEDDATA_VER		TEXT("E9A61FD4C181422A8EAAB060B1690F6A")
 
+#ifndef CASE_ENUM_TO_TEXT
+#define CASE_ENUM_TO_TEXT(X) case X: return TEXT(#X);
+#endif
+
+const TCHAR* LexToString(const ESoundwaveSampleRateSettings Enum)
+{
+	switch (Enum)
+	{
+		FOREACH_ENUM_ESOUNDWAVESAMPLERATESETTINGS(CASE_ENUM_TO_TEXT)
+	}
+	return TEXT("<Unknown ESoundwaveSampleRateSettings>");
+}
+
 /**
  * Computes the derived data key suffix for a SoundWave's Streamed Audio.
  * @param SoundWave - The SoundWave for which to compute the derived data key.
@@ -200,9 +214,20 @@ static void GetStreamedAudioDerivedDataKeySuffix(
 		FPlatformAudioCookOverrides::GetHashSuffix(CompressionOverrides, AudioFormatNameString);
 	}
 
+	// Hash the parts of the SoundWave that can affect the compressed data. It doesn't hurt to do this. 
+	// Typically the GUID will change if compressed data changes, but some settings can affect 
+	// the SoundWave. i.e. DefaultSoundWaveQuality which won't change the GUID. So is better
+	// it's reflected here.
+	using FPCU = FPlatformCompressionUtilities;
+	FString SoundWaveHash;
+	FPCU::AppendHash(SoundWaveHash, TEXT("QLT"), SoundWave.GetCompressionQuality());
+	FPCU::AppendHash(SoundWaveHash, TEXT("CHN"), SoundWave.NumChannels);
+	FPCU::AppendHash(SoundWaveHash, TEXT("SRQ"), SoundWave.SampleRateQuality);
+		
 	// build the key
-	OutKeySuffix = FString::Printf(TEXT("%s_%d_%s"),
+	OutKeySuffix = FString::Printf(TEXT("%s_%s_%d_%s"),
 		*AudioFormatNameString,
+		*SoundWaveHash,
 		Version,
 		*SoundWave.CompressedDataGuid.ToString()
 		);
diff --git a/Engine/Source/Runtime/Engine/Public/AudioCompressionSettingsUtils.h b/Engine/Source/Runtime/Engine/Public/AudioCompressionSettingsUtils.h
--- a/Engine/Source/Runtime/Engine/Public/AudioCompressionSettingsUtils.h
+++ b/Engine/Source/Runtime/Engine/Public/AudioCompressionSettingsUtils.h
@@ -1,9 +1,5 @@
 // Copyright Epic Games, Inc. All Rights Reserved.
 #pragma once
-
-#include "CoreMinimal.h"
-#include "IAudioExtensionPlugin.h"
-#include "AudioPluginUtilities.h"
 #include "AudioCompressionSettings.h"
 
 #if UE_ENABLE_INCLUDE_ORDER_DEPRECATED_IN_5_2
@@ -41,6 +37,12 @@ public:
 	/** This is used at runtime in BuildCachedStreamingManagerParams, as well as cooktime in FStreamedAudioCacheDerivedDataWorker::BuildStreamedAudio to split compressed audio.  */
 	static uint32 GetMaxChunkSizeForCookOverrides(const FPlatformAudioCookOverrides* InCompressionOverrides);
 
+	template<typename HashType>
+	static void AppendHash(FString& OutString, const TCHAR* InName, const HashType& InValueToHash)
+	{
+		OutString += FString::Printf(TEXT("%s_%s_"), InName, ToCStr(LexToString(InValueToHash)));
+	}
+
 private:
 	static const FPlatformRuntimeAudioCompressionOverrides* GetRuntimeCompressionOverridesForCurrentPlatform();
 	
