fix bug in audio resampling which was causing automation test failures in //MetaHuman/Dev-MHPlugin-EngineMerge

diff --git a/Engine/Source/Runtime/AudioPlatformConfiguration/Private/AudioResampler.cpp b/Engine/Source/Runtime/AudioPlatformConfiguration/Private/AudioResampler.cpp
--- a/Engine/Source/Runtime/AudioPlatformConfiguration/Private/AudioResampler.cpp
+++ b/Engine/Source/Runtime/AudioPlatformConfiguration/Private/AudioResampler.cpp
@@ -77,6 +77,7 @@ namespace Audio
 		SrcData.input_frames = InParameters.InputBuffer.Num() / InParameters.NumChannels;
 		SrcData.output_frames = OutData.OutBuffer->Num() / InParameters.NumChannels;
 		SrcData.src_ratio = InParameters.DestinationSampleRate / InParameters.SourceSampleRate;
+		SrcData.end_of_input = 1;
 
 		Error = src_process(Converter, &SrcData);
 
