Remove "Match Device" from Quality Dropdown on Sound Waves
#rnx

diff --git a/Engine/Source/Runtime/AudioPlatformConfiguration/Public/AudioCompressionSettings.h b/Engine/Source/Runtime/AudioPlatformConfiguration/Public/AudioCompressionSettings.h
--- a/Engine/Source/Runtime/AudioPlatformConfiguration/Public/AudioCompressionSettings.h
+++ b/Engine/Source/Runtime/AudioPlatformConfiguration/Public/AudioCompressionSettings.h
@@ -17,8 +17,7 @@ enum class ESoundwaveSampleRateSettings : uint8
 	High,
 	Medium,
 	Low,
-	Min,
-	MatchDevice_DEPRECATED
+	Min
 };
 
 /************************************************************************/
